Fixed non unity compile error

diff --git a/Engine/Plugins/Editor/MobileLauncherProfileWizard/Source/MobileLauncherProfileWizard/Private/SProfileWizardUI.cpp b/Engine/Plugins/Editor/MobileLauncherProfileWizard/Source/MobileLauncherProfileWizard/Private/SProfileWizardUI.cpp
--- a/Engine/Plugins/Editor/MobileLauncherProfileWizard/Source/MobileLauncherProfileWizard/Private/SProfileWizardUI.cpp
+++ b/Engine/Plugins/Editor/MobileLauncherProfileWizard/Source/MobileLauncherProfileWizard/Private/SProfileWizardUI.cpp
@@ -9,6 +9,7 @@
 #include "InstalledPlatformInfo.h"
 #include "PlatformInfo.h"
 #include "DesktopPlatformModule.h"
+#include "GameProjectHelper.h"
 
 #define LOCTEXT_NAMESPACE "MobileLauncherProfileWizard"
 
