Fix incorrect text case when referencing modules

#rnx

diff --git a/Engine/Plugins/Editor/GLTFImporter/Source/GLTFImporter/GLTFImporter.Build.cs b/Engine/Plugins/Editor/GLTFImporter/Source/GLTFImporter/GLTFImporter.Build.cs
--- a/Engine/Plugins/Editor/GLTFImporter/Source/GLTFImporter/GLTFImporter.Build.cs
+++ b/Engine/Plugins/Editor/GLTFImporter/Source/GLTFImporter/GLTFImporter.Build.cs
@@ -38,7 +38,7 @@ namespace UnrealBuildTool.Rules
                     "MaterialEditor",
                     "Slate",
                     "SlateCore",
-                    "Mainframe",
+                    "MainFrame",
                     "InputCore",
                     "EditorStyle",
 					"GLTFCore",
diff --git a/Engine/Plugins/Editor/MobileLauncherProfileWizard/Source/MobileLauncherProfileWizard/MobileLauncherProfileWizard.Build.cs b/Engine/Plugins/Editor/MobileLauncherProfileWizard/Source/MobileLauncherProfileWizard/MobileLauncherProfileWizard.Build.cs
--- a/Engine/Plugins/Editor/MobileLauncherProfileWizard/Source/MobileLauncherProfileWizard/MobileLauncherProfileWizard.Build.cs
+++ b/Engine/Plugins/Editor/MobileLauncherProfileWizard/Source/MobileLauncherProfileWizard/MobileLauncherProfileWizard.Build.cs
@@ -33,7 +33,7 @@ namespace UnrealBuildTool.Rules
                     "CoreUObject",
                     "TargetPlatform",
                     "DesktopPlatform",
-                    "JSON",
+                    "Json",
                     "Slate",
                     "SlateCore",
                     "InputCore",
diff --git a/Engine/Plugins/Enterprise/AxFImporter/Source/AxFImporter/AxFImporter.Build.cs b/Engine/Plugins/Enterprise/AxFImporter/Source/AxFImporter/AxFImporter.Build.cs
--- a/Engine/Plugins/Enterprise/AxFImporter/Source/AxFImporter/AxFImporter.Build.cs
+++ b/Engine/Plugins/Enterprise/AxFImporter/Source/AxFImporter/AxFImporter.Build.cs
@@ -31,7 +31,7 @@ namespace UnrealBuildTool.Rules
                     "UnrealEd",
                     "Slate",
                     "SlateCore",
-                    "Mainframe",
+                    "MainFrame",
                     "InputCore",
                     "EditorStyle",
                     "MaterialEditor",
diff --git a/Engine/Plugins/Enterprise/MDLImporter/Source/MDLImporter/MDLImporter.Build.cs b/Engine/Plugins/Enterprise/MDLImporter/Source/MDLImporter/MDLImporter.Build.cs
--- a/Engine/Plugins/Enterprise/MDLImporter/Source/MDLImporter/MDLImporter.Build.cs
+++ b/Engine/Plugins/Enterprise/MDLImporter/Source/MDLImporter/MDLImporter.Build.cs
@@ -33,7 +33,7 @@ namespace UnrealBuildTool.Rules
                     "UnrealEd",
                     "Slate",
                     "SlateCore",
-                    "Mainframe",
+                    "MainFrame",
                     "InputCore",
                     "EditorStyle",
                     "MaterialEditor",
diff --git a/Engine/Plugins/Experimental/AutomationUtils/Source/AutomationUtilsEditor/AutomationUtilsEditor.Build.cs b/Engine/Plugins/Experimental/AutomationUtils/Source/AutomationUtilsEditor/AutomationUtilsEditor.Build.cs
--- a/Engine/Plugins/Experimental/AutomationUtils/Source/AutomationUtilsEditor/AutomationUtilsEditor.Build.cs
+++ b/Engine/Plugins/Experimental/AutomationUtils/Source/AutomationUtilsEditor/AutomationUtilsEditor.Build.cs
@@ -15,8 +15,8 @@ namespace UnrealBuildTool.Rules
 					"ScreenshotTools",
 					"DesktopPlatform",
 					"AutomationMessages",
-					"JSON",
-					"JSONUtilities"
+					"Json",
+					"JsonUtilities"
 				}
 			);
 		}
diff --git a/Engine/Plugins/Importers/USDMultiUser/Source/USDMultiUser/USDMultiUser.Build.cs b/Engine/Plugins/Importers/USDMultiUser/Source/USDMultiUser/USDMultiUser.Build.cs
--- a/Engine/Plugins/Importers/USDMultiUser/Source/USDMultiUser/USDMultiUser.Build.cs
+++ b/Engine/Plugins/Importers/USDMultiUser/Source/USDMultiUser/USDMultiUser.Build.cs
@@ -20,8 +20,8 @@ public class USDMultiUser : ModuleRules
 				"Engine",
 				"UnrealUSDWrapper",
 				"USDClasses",
-				"UsdStage",
-				"UsdUtilities",
+				"USDStage",
+				"USDUtilities",
 			}
 		);
 	}
diff --git a/Engine/Plugins/MovieScene/MovieRenderPipeline/Source/MovieRenderPipelineCore/MovieRenderPipelineCore.Build.cs b/Engine/Plugins/MovieScene/MovieRenderPipeline/Source/MovieRenderPipelineCore/MovieRenderPipelineCore.Build.cs
--- a/Engine/Plugins/MovieScene/MovieRenderPipeline/Source/MovieRenderPipelineCore/MovieRenderPipelineCore.Build.cs
+++ b/Engine/Plugins/MovieScene/MovieRenderPipeline/Source/MovieRenderPipelineCore/MovieRenderPipelineCore.Build.cs
@@ -19,7 +19,7 @@ public class MovieRenderPipelineCore : ModuleRules
 				"NonRealtimeAudioRenderer",
 				"Sockets", 
 				"Networking",
-				"Http",
+				"HTTP",
 				"DeveloperSettings",
 				"CinematicCamera"
 
diff --git a/Engine/Plugins/Online/Android/OnlineSubsystemGameCircle/Source/OnlineSubsystemGameCircle/OnlineSubsystemGameCircle.Build.cs b/Engine/Plugins/Online/Android/OnlineSubsystemGameCircle/Source/OnlineSubsystemGameCircle/OnlineSubsystemGameCircle.Build.cs
--- a/Engine/Plugins/Online/Android/OnlineSubsystemGameCircle/Source/OnlineSubsystemGameCircle/OnlineSubsystemGameCircle.Build.cs
+++ b/Engine/Plugins/Online/Android/OnlineSubsystemGameCircle/Source/OnlineSubsystemGameCircle/OnlineSubsystemGameCircle.Build.cs
@@ -25,7 +25,7 @@ namespace UnrealBuildTool.Rules
 				"Engine",
 				"Sockets",
 				"OnlineSubsystem",
-				"Http",
+				"HTTP",
 				"GameCircleRuntimeSettings",
 				"Launch"
 				}
diff --git a/Engine/Plugins/Online/Android/OnlineSubsystemGooglePlay/Source/OnlineSubsystemGooglePlay.Build.cs b/Engine/Plugins/Online/Android/OnlineSubsystemGooglePlay/Source/OnlineSubsystemGooglePlay.Build.cs
--- a/Engine/Plugins/Online/Android/OnlineSubsystemGooglePlay/Source/OnlineSubsystemGooglePlay.Build.cs
+++ b/Engine/Plugins/Online/Android/OnlineSubsystemGooglePlay/Source/OnlineSubsystemGooglePlay.Build.cs
@@ -31,7 +31,7 @@ public class OnlineSubsystemGooglePlay : ModuleRules
                 "Engine", 
                 "Sockets",
 				"OnlineSubsystem", 
-                "Http",
+                "HTTP",
 				"AndroidRuntimeSettings",
 				"Launch",
 				"GpgCppSDK"
diff --git a/Engine/Plugins/Online/IOS/OnlineSubsystemIOS/Source/OnlineSubsystemIOS.Build.cs b/Engine/Plugins/Online/IOS/OnlineSubsystemIOS/Source/OnlineSubsystemIOS.Build.cs
--- a/Engine/Plugins/Online/IOS/OnlineSubsystemIOS/Source/OnlineSubsystemIOS.Build.cs
+++ b/Engine/Plugins/Online/IOS/OnlineSubsystemIOS/Source/OnlineSubsystemIOS.Build.cs
@@ -26,7 +26,7 @@ public class OnlineSubsystemIOS : ModuleRules
                 "Engine", 
                 "Sockets",
 				"OnlineSubsystem", 
-                "Http",
+                "HTTP",
                 "IOSPlatformFeatures",
             }
             );
diff --git a/Engine/Plugins/Runtime/HTTPChunkInstaller/Source/HTTPChunkInstaller.Build.cs b/Engine/Plugins/Runtime/HTTPChunkInstaller/Source/HTTPChunkInstaller.Build.cs
--- a/Engine/Plugins/Runtime/HTTPChunkInstaller/Source/HTTPChunkInstaller.Build.cs
+++ b/Engine/Plugins/Runtime/HTTPChunkInstaller/Source/HTTPChunkInstaller.Build.cs
@@ -13,7 +13,7 @@ public class HTTPChunkInstaller : ModuleRules
                 "Core",
 				"ApplicationCore",
                 "Engine",
-                "Http",
+                "HTTP",
                 "Json",
                 "PakFile",
             }
diff --git a/Engine/Source/Developer/LauncherServices/LauncherServices.Build.cs b/Engine/Source/Developer/LauncherServices/LauncherServices.Build.cs
--- a/Engine/Source/Developer/LauncherServices/LauncherServices.Build.cs
+++ b/Engine/Source/Developer/LauncherServices/LauncherServices.Build.cs
@@ -24,7 +24,7 @@ namespace UnrealBuildTool.Rules
 					"SourceCodeAccess",
 					"TargetPlatform",
 					"UnrealEdMessages",
-					"JSON",
+					"Json",
 					"TurnkeyIO",
 				});
 
diff --git a/Engine/Source/Editor/UnrealEd/UnrealEd.Build.cs b/Engine/Source/Editor/UnrealEd/UnrealEd.Build.cs
--- a/Engine/Source/Editor/UnrealEd/UnrealEd.Build.cs
+++ b/Engine/Source/Editor/UnrealEd/UnrealEd.Build.cs
@@ -93,7 +93,7 @@ public class UnrealEd : ModuleRules
 				"UnrealEdMessages",
 				"GameplayDebugger",
 				"BlueprintGraph",
-				"Http",
+				"HTTP",
 				"UnrealAudio",
 				"FunctionalTesting",
 				"AutomationController",
diff --git a/Engine/Source/Programs/BuildPatchTool/BuildPatchTool.Build.cs b/Engine/Source/Programs/BuildPatchTool/BuildPatchTool.Build.cs
--- a/Engine/Source/Programs/BuildPatchTool/BuildPatchTool.Build.cs
+++ b/Engine/Source/Programs/BuildPatchTool/BuildPatchTool.Build.cs
@@ -19,7 +19,7 @@ public class BuildPatchTool : ModuleRules
 				"CoreUObject",
 				"BuildPatchServices",
 				"Projects",
-				"Http",
+				"HTTP",
 				// The below items are not strictly needed by BPT, but core appears to need them during initialization
 				"PakFile",
 				"SandboxFile",
