Fixed UI issue where iOS DLC wizard does not progress past the cook flavor

diff --git a/Engine/Plugins/Editor/MobileLauncherProfileWizard/Source/MobileLauncherProfileWizard/Private/SProfileWizardUI.cpp b/Engine/Plugins/Editor/MobileLauncherProfileWizard/Source/MobileLauncherProfileWizard/Private/SProfileWizardUI.cpp
--- a/Engine/Plugins/Editor/MobileLauncherProfileWizard/Source/MobileLauncherProfileWizard/Private/SProfileWizardUI.cpp
+++ b/Engine/Plugins/Editor/MobileLauncherProfileWizard/Source/MobileLauncherProfileWizard/Private/SProfileWizardUI.cpp
@@ -744,8 +744,12 @@ void SProfileWizardUI::CacheCookFlavorsList()
 	if (DLCFlavorList.Num() == 0)
 	{
 		DLCFlavorList.Add(MakeShareable(new FString(TargetPlatformName.ToString())));
-		// And make it selected by default
-		DLCSelectedFlavors.Add(TargetPlatformName.ToString());
+	}
+
+	// If only one target, make it selected by default
+	if (DLCFlavorList.Num() == 1)
+	{
+		DLCSelectedFlavors.Add(*DLCFlavorList[0].Get());
 	}
 }
 
